Added a file logging analytics provider

diff --git a/Engine/Plugins/Runtime/Analytics/FileLogging/FileLogging.uplugin b/Engine/Plugins/Runtime/Analytics/FileLogging/FileLogging.uplugin
new file mode 100644
--- /dev/null
+++ b/Engine/Plugins/Runtime/Analytics/FileLogging/FileLogging.uplugin
@@ -0,0 +1,21 @@
+{
+	"FileVersion" : 3,
+	
+	"FriendlyName" : "File Logging Analytics Provider",
+	"Version" : 1,
+	"VersionName" : "1.0",
+	"CreatedBy" : "Epic Games, Inc.",
+	"CreatedByURL" : "http://epicgames.com",
+	"EngineVersion" : "4.2.0",
+	"Description" : "Writes analytic API calls to local disk for debugging or local use",
+	"Category" : "Analytics.Provider",
+	
+	"Modules" :
+	[
+		{
+			"Name" : "FileLogging",
+			"Type" : "Runtime",
+			"LoadingPhase" : "PreDefault"
+		}
+	]
+}
\ No newline at end of file
diff --git a/Engine/Plugins/Runtime/Analytics/FileLogging/Source/FileLogging/FileLogging.Build.cs b/Engine/Plugins/Runtime/Analytics/FileLogging/Source/FileLogging/FileLogging.Build.cs
new file mode 100644
--- /dev/null
+++ b/Engine/Plugins/Runtime/Analytics/FileLogging/Source/FileLogging/FileLogging.Build.cs
@@ -0,0 +1,30 @@
+// Copyright 1998-2014 Epic Games, Inc. All Rights Reserved.
+
+namespace UnrealBuildTool.Rules
+{
+	public class FileLogging : ModuleRules
+	{
+		public FileLogging( TargetInfo Target )
+		{
+            PublicDependencyModuleNames.AddRange(
+                new string[]
+				{
+					"Core",
+					"CoreUObject",
+                    "Engine",
+					// ... add other public dependencies that you statically link with here ...
+				}
+                );
+
+            PrivateDependencyModuleNames.AddRange(
+                new string[]
+				{
+					"Analytics",
+					// ... add private dependencies that you statically link with here ...
+				}
+                );
+
+            PublicIncludePathModuleNames.Add("Analytics");
+        }
+	}
+}
\ No newline at end of file
diff --git a/Engine/Plugins/Runtime/Analytics/FileLogging/Source/FileLogging/Private/FileLogging.cpp b/Engine/Plugins/Runtime/Analytics/FileLogging/Source/FileLogging/Private/FileLogging.cpp
new file mode 100644
--- /dev/null
+++ b/Engine/Plugins/Runtime/Analytics/FileLogging/Source/FileLogging/Private/FileLogging.cpp
@@ -0,0 +1,269 @@
+// Copyright 1998-2014 Epic Games, Inc. All Rights Reserved.
+
+#include "FileLoggingPrivatePCH.h"
+
+DEFINE_LOG_CATEGORY_STATIC(LogFileLoggingAnalytics, Display, All);
+
+IMPLEMENT_MODULE( FAnalyticsFileLogging, FileLogging )
+
+void FAnalyticsFileLogging::StartupModule()
+{
+	FileLoggingProvider = MakeShareable(new FAnalyticsProviderFileLogging());
+}
+
+void FAnalyticsFileLogging::ShutdownModule()
+{
+	if (FileLoggingProvider.IsValid())
+	{
+		FileLoggingProvider->EndSession();
+	}
+}
+
+TSharedPtr<IAnalyticsProvider> FAnalyticsFileLogging::CreateAnalyticsProvider(const FAnalytics::FProviderConfigurationDelegate& GetConfigValue) const
+{
+	return FileLoggingProvider;
+}
+
+// Provider
+
+FAnalyticsProviderFileLogging::FAnalyticsProviderFileLogging() :
+	bHasSessionStarted(false),
+	bHasWrittenFirstEvent(false),
+	FileArchive(nullptr)
+{
+	FileArchive = nullptr;
+	AnalyticsFilePath = FPaths::GameSavedDir() + TEXT("Analytics/");
+	UserId = FPlatformMisc::GetUniqueDeviceId();
+}
+
+FAnalyticsProviderFileLogging::~FAnalyticsProviderFileLogging()
+{
+	if (bHasSessionStarted)
+	{
+		EndSession();
+	}
+}
+
+bool FAnalyticsProviderFileLogging::StartSession(const TArray<FAnalyticsEventAttribute>& Attributes)
+{
+	if (bHasSessionStarted)
+	{
+		EndSession();
+	}
+	SessionId = UserId + TEXT("-") + FDateTime::Now().ToString();
+	const FString FileName = AnalyticsFilePath + SessionId + TEXT(".analytics");
+	// Close the old file and open a new one
+	FileArchive = IFileManager::Get().CreateFileWriter(*FileName);
+	if (FileArchive != nullptr)
+	{
+		FileArchive->Logf(TEXT("{"));
+		FileArchive->Logf(TEXT("\t\"sessionId\" : \"%s\","), *SessionId);
+		FileArchive->Logf(TEXT("\t\"userId\" : \"%s\","), *UserId);
+		FileArchive->Logf(TEXT("\t\"events\" : ["));
+		bHasSessionStarted = true;
+	}
+	else
+	{
+		UE_LOG(LogFileLoggingAnalytics, Warning, TEXT("FAnalyticsProviderFileLogging::StartSession failed to create file to log analytics events to"));
+	}
+	return bHasSessionStarted;
+}
+
+void FAnalyticsProviderFileLogging::EndSession()
+{
+	if (FileArchive != nullptr)
+	{
+		FileArchive->Logf(TEXT("\t]"));
+		FileArchive->Logf(TEXT("}"));
+		FileArchive->Flush();
+		FileArchive->Close();
+		delete FileArchive;
+		FileArchive = nullptr;
+	}
+	bHasWrittenFirstEvent = false;
+	bHasSessionStarted = false;
+}
+
+void FAnalyticsProviderFileLogging::FlushEvents()
+{
+	if (FileArchive != nullptr)
+	{
+		FileArchive->Flush();
+	}
+}
+
+void FAnalyticsProviderFileLogging::SetUserID(const FString& InUserID)
+{
+	if (!bHasSessionStarted)
+	{
+		UserId = InUserID;
+	}
+	else
+	{
+		// Log that we shouldn't switch users during a session
+		UE_LOG(LogFileLoggingAnalytics, Warning, TEXT("FAnalyticsProviderFileLogging::SetUserID called while a session is in progress. Ignoring."));
+	}
+}
+
+FString FAnalyticsProviderFileLogging::GetUserID() const
+{
+	return UserId;
+}
+
+FString FAnalyticsProviderFileLogging::GetSessionID() const
+{
+	return SessionId;
+}
+
+bool FAnalyticsProviderFileLogging::SetSessionID(const FString& InSessionID)
+{
+	if (!bHasSessionStarted)
+	{
+		SessionId = InSessionID;
+	}
+	else
+	{
+		// Log that we shouldn't switch session ids during a session
+		UE_LOG(LogFileLoggingAnalytics, Warning, TEXT("FAnalyticsProviderFileLogging::SetSessionID called while a session is in progress. Ignoring."));
+	}
+	return !bHasSessionStarted;
+}
+
+void FAnalyticsProviderFileLogging::RecordEvent(const FString& EventName, const TArray<FAnalyticsEventAttribute>& Attributes)
+{
+	if (bHasSessionStarted)
+	{
+		check(FileArchive != nullptr);
+
+		if (bHasWrittenFirstEvent)
+		{
+			FileArchive->Logf(TEXT(","));
+		}
+		bHasWrittenFirstEvent = true;
+
+		FileArchive->Logf(TEXT("\t\t{"));
+		FileArchive->Logf(TEXT("\t\t\t\"eventName\" : \"%s\""), *EventName);
+		if (Attributes.Num() > 0)
+		{
+			FileArchive->Logf(TEXT(",\t\t\t\"attributes\" : ["));
+			bool bHasWrittenFirstAttr = false;
+			// Write out the list of attributes as an array of attribute objects
+			for (auto Attr : Attributes)
+			{
+				if (bHasWrittenFirstAttr)
+				{
+					FileArchive->Logf(TEXT("\t\t\t,"));
+				}
+				FileArchive->Logf(TEXT("\t\t\t{"));
+				FileArchive->Logf(TEXT("\t\t\t\t\"name\" : \"%s\","), *Attr.AttrName);
+				FileArchive->Logf(TEXT("\t\t\t\t\"value\" : \"%s\""), *Attr.AttrValue);
+				FileArchive->Logf(TEXT("\t\t\t}"));
+				bHasWrittenFirstAttr = true;
+			}
+			FileArchive->Logf(TEXT("\t\t\t]"));
+		}
+		FileArchive->Logf(TEXT("\t\t}"));
+	}
+	else
+	{
+		UE_LOG(LogFileLoggingAnalytics, Warning, TEXT("FAnalyticsProviderFileLogging::RecordEvent called before StartSession. Ignoring."));
+	}
+}
+
+void FAnalyticsProviderFileLogging::RecordItemPurchase(const FString& ItemId, const FString& Currency, int PerItemCost, int ItemQuantity)
+{
+	if (bHasSessionStarted)
+	{
+		check(FileArchive != nullptr);
+
+		if (bHasWrittenFirstEvent)
+		{
+			FileArchive->Logf(TEXT("\t\t,"));
+		}
+		bHasWrittenFirstEvent = true;
+
+		FileArchive->Logf(TEXT("\t\t{"));
+		FileArchive->Logf(TEXT("\t\t\t\"eventName\" : \"recordItemPurchase\","));
+
+		FileArchive->Logf(TEXT("\t\t\t\"attributes\" :"));
+		FileArchive->Logf(TEXT("\t\t\t["));
+
+		FileArchive->Logf(TEXT("\t\t\t\t{ \"name\" : \"itemId\", \t\"value\" : \"%s\" },"), *ItemId);
+		FileArchive->Logf(TEXT("\t\t\t\t{ \"name\" : \"currency\", \t\"value\" : \"%s\" },"), *Currency);
+		FileArchive->Logf(TEXT("\t\t\t\t{ \"name\" : \"perItemCost\", \t\"value\" : \"%d\" },"), PerItemCost);
+		FileArchive->Logf(TEXT("\t\t\t\t{ \"name\" : \"itemQuantity\", \t\"value\" : \"%d\" }"), ItemQuantity);
+
+		FileArchive->Logf(TEXT("\t\t\t]"));
+
+		FileArchive->Logf(TEXT("\t\t}"));
+	}
+	else
+	{
+		UE_LOG(LogFileLoggingAnalytics, Warning, TEXT("FAnalyticsProviderFileLogging::RecordItemPurchase called before StartSession. Ignoring."));
+	}
+}
+
+void FAnalyticsProviderFileLogging::RecordCurrencyPurchase(const FString& GameCurrencyType, int GameCurrencyAmount, const FString& RealCurrencyType, float RealMoneyCost, const FString& PaymentProvider)
+{
+	if (bHasSessionStarted)
+	{
+		check(FileArchive != nullptr);
+
+		if (bHasWrittenFirstEvent)
+		{
+			FileArchive->Logf(TEXT("\t\t,"));
+		}
+		bHasWrittenFirstEvent = true;
+
+		FileArchive->Logf(TEXT("\t\t{"));
+		FileArchive->Logf(TEXT("\t\t\t\"eventName\" : \"recordCurrencyPurchase\","));
+
+		FileArchive->Logf(TEXT("\t\t\t\"attributes\" :"));
+		FileArchive->Logf(TEXT("\t\t\t["));
+
+		FileArchive->Logf(TEXT("\t\t\t\t{ \"name\" : \"gameCurrencyType\", \t\"value\" : \"%s\" },"), *GameCurrencyType);
+		FileArchive->Logf(TEXT("\t\t\t\t{ \"name\" : \"gameCurrencyAmount\", \t\"value\" : \"%d\" },"), GameCurrencyAmount);
+		FileArchive->Logf(TEXT("\t\t\t\t{ \"name\" : \"realCurrencyType\", \t\"value\" : \"%s\" },"), *RealCurrencyType);
+		FileArchive->Logf(TEXT("\t\t\t\t{ \"name\" : \"realMoneyCost\", \t\"value\" : \"%f\" },"), RealMoneyCost);
+		FileArchive->Logf(TEXT("\t\t\t\t{ \"name\" : \"paymentProvider\", \t\"value\" : \"%s\" }"), *PaymentProvider);
+
+		FileArchive->Logf(TEXT("\t\t\t]"));
+
+		FileArchive->Logf(TEXT("\t\t}"));
+	}
+	else
+	{
+		UE_LOG(LogFileLoggingAnalytics, Warning, TEXT("FAnalyticsProviderFileLogging::RecordCurrencyPurchase called before StartSession. Ignoring."));
+	}
+}
+
+void FAnalyticsProviderFileLogging::RecordCurrencyGiven(const FString& GameCurrencyType, int GameCurrencyAmount)
+{
+	if (bHasSessionStarted)
+	{
+		check(FileArchive != nullptr);
+
+		if (bHasWrittenFirstEvent)
+		{
+			FileArchive->Logf(TEXT("\t\t,"));
+		}
+		bHasWrittenFirstEvent = true;
+
+		FileArchive->Logf(TEXT("\t\t{"));
+		FileArchive->Logf(TEXT("\t\t\t\"eventName\" : \"recordCurrencyGiven\","));
+
+		FileArchive->Logf(TEXT("\t\t\t\"attributes\" :"));
+		FileArchive->Logf(TEXT("\t\t\t["));
+
+		FileArchive->Logf(TEXT("\t\t\t\t{ \"name\" : \"gameCurrencyType\", \t\"value\" : \"%s\" },"), *GameCurrencyType);
+		FileArchive->Logf(TEXT("\t\t\t\t{ \"name\" : \"gameCurrencyAmount\", \t\"value\" : \"%d\" }"), GameCurrencyAmount);
+
+		FileArchive->Logf(TEXT("\t\t\t]"));
+
+		FileArchive->Logf(TEXT("\t\t}"));
+	}
+	else
+	{
+		UE_LOG(LogFileLoggingAnalytics, Warning, TEXT("FAnalyticsProviderFileLogging::RecordCurrencyGiven called before StartSession. Ignoring."));
+	}
+}
diff --git a/Engine/Plugins/Runtime/Analytics/FileLogging/Source/FileLogging/Private/FileLoggingPrivatePCH.h b/Engine/Plugins/Runtime/Analytics/FileLogging/Source/FileLogging/Private/FileLoggingPrivatePCH.h
new file mode 100644
--- /dev/null
+++ b/Engine/Plugins/Runtime/Analytics/FileLogging/Source/FileLogging/Private/FileLoggingPrivatePCH.h
@@ -0,0 +1,11 @@
+// Copyright 1998-2014 Epic Games, Inc. All Rights Reserved.
+
+#include "Core.h"
+#include "IAnalyticsProviderModule.h"
+#include "IAnalyticsProvider.h"
+
+#include "FileLogging.h"
+#include "FileLoggingProvider.h"
+
+// You should place include statements to your module's private header files here.  You only need to
+// add includes for headers that are used in most of your module's source files though.
diff --git a/Engine/Plugins/Runtime/Analytics/FileLogging/Source/FileLogging/Private/FileLoggingProvider.h b/Engine/Plugins/Runtime/Analytics/FileLogging/Source/FileLogging/Private/FileLoggingProvider.h
new file mode 100644
--- /dev/null
+++ b/Engine/Plugins/Runtime/Analytics/FileLogging/Source/FileLogging/Private/FileLoggingProvider.h
@@ -0,0 +1,40 @@
+// Copyright 1998-2014 Epic Games, Inc. All Rights Reserved.
+
+#pragma once
+
+class FAnalyticsProviderFileLogging :
+	public IAnalyticsProvider
+{
+	/** Path where analytics files are saved out */
+	FString AnalyticsFilePath;
+	/** Tracks whether we need to start the session or restart it */
+	bool bHasSessionStarted;
+	/** Whether an event was written before or not */
+	bool bHasWrittenFirstEvent;
+	/** Id representing the user the analytics are recording for */
+	FString UserId;
+	/** Unique Id representing the session the analytics are recording for */
+	FString SessionId;
+	/** The file archive used to write the data */
+	FArchive* FileArchive;
+
+public:
+	FAnalyticsProviderFileLogging();
+	virtual ~FAnalyticsProviderFileLogging();
+
+	virtual bool StartSession(const TArray<FAnalyticsEventAttribute>& Attributes) override;
+	virtual void EndSession() override;
+	virtual void FlushEvents() override;
+
+	virtual void SetUserID(const FString& InUserID) override;
+	virtual FString GetUserID() const override;
+
+	virtual FString GetSessionID() const override;
+	virtual bool SetSessionID(const FString& InSessionID) override;
+
+	virtual void RecordEvent(const FString& EventName, const TArray<FAnalyticsEventAttribute>& Attributes) override;
+
+	virtual void RecordItemPurchase(const FString& ItemId, const FString& Currency, int PerItemCost, int ItemQuantity) override;
+	virtual void RecordCurrencyPurchase(const FString& GameCurrencyType, int GameCurrencyAmount, const FString& RealCurrencyType, float RealMoneyCost, const FString& PaymentProvider) override;
+	virtual void RecordCurrencyGiven(const FString& GameCurrencyType, int GameCurrencyAmount) override;
+};
diff --git a/Engine/Plugins/Runtime/Analytics/FileLogging/Source/FileLogging/Public/FileLogging.h b/Engine/Plugins/Runtime/Analytics/FileLogging/Source/FileLogging/Public/FileLogging.h
new file mode 100644
--- /dev/null
+++ b/Engine/Plugins/Runtime/Analytics/FileLogging/Source/FileLogging/Public/FileLogging.h
@@ -0,0 +1,49 @@
+// Copyright 1998-2014 Epic Games, Inc. All Rights Reserved.
+
+#pragma once
+
+#include "Runtime/Analytics/Analytics/Public/Interfaces/IAnalyticsProviderModule.h"
+#include "Core.h"
+
+class IAnalyticsProvider;
+
+/**
+ * The public interface to this module
+ */
+class FAnalyticsFileLogging :
+	public IAnalyticsProviderModule
+{
+	/** Singleton for analytics */
+	TSharedPtr<IAnalyticsProvider> FileLoggingProvider;
+
+	//--------------------------------------------------------------------------
+	// Module functionality
+	//--------------------------------------------------------------------------
+public:
+	/**
+	 * Singleton-like access to this module's interface.  This is just for convenience!
+	 * Beware of calling this during the shutdown phase, though.  Your module might have been unloaded already.
+	 *
+	 * @return Returns singleton instance, loading the module on demand if needed
+	 */
+	static inline FAnalyticsFileLogging& Get()
+	{
+		return FModuleManager::LoadModuleChecked< FAnalyticsFileLogging >( "FileLogging" );
+	}
+
+	//--------------------------------------------------------------------------
+	// provider factory functions
+	//--------------------------------------------------------------------------
+public:
+	/**
+	 * IAnalyticsProviderModule interface.
+	 * Creates the analytics provider given a configuration delegate.
+	 * The keys required exactly match the field names in the Config object. 
+	 */
+	virtual TSharedPtr<IAnalyticsProvider> CreateAnalyticsProvider(const FAnalytics::FProviderConfigurationDelegate& GetConfigValue) const override;
+	
+private:
+	virtual void StartupModule() override;
+	virtual void ShutdownModule() override;
+};
+
