Added logging to the console of the data written using the file logging analytics provider

diff --git a/Engine/Plugins/Runtime/Analytics/FileLogging/Source/FileLogging/Private/FileLogging.cpp b/Engine/Plugins/Runtime/Analytics/FileLogging/Source/FileLogging/Private/FileLogging.cpp
--- a/Engine/Plugins/Runtime/Analytics/FileLogging/Source/FileLogging/Private/FileLogging.cpp
+++ b/Engine/Plugins/Runtime/Analytics/FileLogging/Source/FileLogging/Private/FileLogging.cpp
@@ -61,6 +61,7 @@ bool FAnalyticsProviderFileLogging::StartSession(const TArray<FAnalyticsEventAtt
 		FileArchive->Logf(TEXT("\t\"userId\" : \"%s\","), *UserId);
 		FileArchive->Logf(TEXT("\t\"events\" : ["));
 		bHasSessionStarted = true;
+		UE_LOG(LogFileLoggingAnalytics, Display, TEXT("Session created file (%s) for user (%s)"), *FileName, *UserId);
 	}
 	else
 	{
@@ -79,6 +80,7 @@ void FAnalyticsProviderFileLogging::EndSession()
 		FileArchive->Close();
 		delete FileArchive;
 		FileArchive = nullptr;
+		UE_LOG(LogFileLoggingAnalytics, Display, TEXT("Session ended for user (%s) and session id (%s)"), *UserId, *SessionId);
 	}
 	bHasWrittenFirstEvent = false;
 	bHasSessionStarted = false;
@@ -89,6 +91,7 @@ void FAnalyticsProviderFileLogging::FlushEvents()
 	if (FileArchive != nullptr)
 	{
 		FileArchive->Flush();
+		UE_LOG(LogFileLoggingAnalytics, Display, TEXT("Analytics file flushed"));
 	}
 }
 
@@ -97,6 +100,7 @@ void FAnalyticsProviderFileLogging::SetUserID(const FString& InUserID)
 	if (!bHasSessionStarted)
 	{
 		UserId = InUserID;
+		UE_LOG(LogFileLoggingAnalytics, Display, TEXT("User is now (%s)"), *UserId);
 	}
 	else
 	{
@@ -120,6 +124,7 @@ bool FAnalyticsProviderFileLogging::SetSessionID(const FString& InSessionID)
 	if (!bHasSessionStarted)
 	{
 		SessionId = InSessionID;
+		UE_LOG(LogFileLoggingAnalytics, Display, TEXT("Session is now (%s)"), *SessionId);
 	}
 	else
 	{
@@ -163,6 +168,8 @@ void FAnalyticsProviderFileLogging::RecordEvent(const FString& EventName, const
 			FileArchive->Logf(TEXT("\t\t\t]"));
 		}
 		FileArchive->Logf(TEXT("\t\t}"));
+
+		UE_LOG(LogFileLoggingAnalytics, Display, TEXT("Analytics event (%s) written with (%d) attributes"), *EventName, Attributes.Num());
 	}
 	else
 	{
@@ -196,6 +203,8 @@ void FAnalyticsProviderFileLogging::RecordItemPurchase(const FString& ItemId, co
 		FileArchive->Logf(TEXT("\t\t\t]"));
 
 		FileArchive->Logf(TEXT("\t\t}"));
+
+		UE_LOG(LogFileLoggingAnalytics, Display, TEXT("(%d) number of item (%s) purchased with (%s) at a cost of (%d) each"), ItemQuantity, *ItemId, *Currency, PerItemCost);
 	}
 	else
 	{
@@ -230,6 +239,8 @@ void FAnalyticsProviderFileLogging::RecordCurrencyPurchase(const FString& GameCu
 		FileArchive->Logf(TEXT("\t\t\t]"));
 
 		FileArchive->Logf(TEXT("\t\t}"));
+
+		UE_LOG(LogFileLoggingAnalytics, Display, TEXT("(%d) amount of in game currency (%s) purchased with (%s) at a cost of (%f) each"), GameCurrencyAmount, *GameCurrencyType, *RealCurrencyType, RealMoneyCost);
 	}
 	else
 	{
@@ -261,6 +272,8 @@ void FAnalyticsProviderFileLogging::RecordCurrencyGiven(const FString& GameCurre
 		FileArchive->Logf(TEXT("\t\t\t]"));
 
 		FileArchive->Logf(TEXT("\t\t}"));
+
+		UE_LOG(LogFileLoggingAnalytics, Display, TEXT("(%d) amount of in game currency (%s) given to user"), GameCurrencyAmount, *GameCurrencyType);
 	}
 	else
 	{
