Updated the file logging and multicast providers to support the new functions

diff --git a/Engine/Plugins/Runtime/Analytics/FileLogging/Source/FileLogging/Private/FileLogging.cpp b/Engine/Plugins/Runtime/Analytics/FileLogging/Source/FileLogging/Private/FileLogging.cpp
--- a/Engine/Plugins/Runtime/Analytics/FileLogging/Source/FileLogging/Private/FileLogging.cpp
+++ b/Engine/Plugins/Runtime/Analytics/FileLogging/Source/FileLogging/Private/FileLogging.cpp
@@ -29,6 +29,7 @@ TSharedPtr<IAnalyticsProvider> FAnalyticsFileLogging::CreateAnalyticsProvider(co
 FAnalyticsProviderFileLogging::FAnalyticsProviderFileLogging() :
 	bHasSessionStarted(false),
 	bHasWrittenFirstEvent(false),
+	Age(0),
 	FileArchive(nullptr)
 {
 	FileArchive = nullptr;
@@ -59,6 +60,22 @@ bool FAnalyticsProviderFileLogging::StartSession(const TArray<FAnalyticsEventAtt
 		FileArchive->Logf(TEXT("{"));
 		FileArchive->Logf(TEXT("\t\"sessionId\" : \"%s\","), *SessionId);
 		FileArchive->Logf(TEXT("\t\"userId\" : \"%s\","), *UserId);
+		if (BuildInfo.Len() > 0)
+		{
+			FileArchive->Logf(TEXT("\t\"buildInfo\" : \"%s\","), *BuildInfo);
+		}
+		if (Age != 0)
+		{
+			FileArchive->Logf(TEXT("\t\"age\" : %d,"), Age);
+		}
+		if (Gender.Len() > 0)
+		{
+			FileArchive->Logf(TEXT("\t\"gender\" : \"%s\","), *Gender);
+		}
+		if (Location.Len() > 0)
+		{
+			FileArchive->Logf(TEXT("\t\"location\" : \"%s\","), *Location);
+		}
 		FileArchive->Logf(TEXT("\t\"events\" : ["));
 		bHasSessionStarted = true;
 		UE_LOG(LogFileLoggingAnalytics, Display, TEXT("Session created file (%s) for user (%s)"), *FileName, *UserId);
@@ -280,3 +297,246 @@ void FAnalyticsProviderFileLogging::RecordCurrencyGiven(const FString& GameCurre
 		UE_LOG(LogFileLoggingAnalytics, Warning, TEXT("FAnalyticsProviderFileLogging::RecordCurrencyGiven called before StartSession. Ignoring."));
 	}
 }
+
+void FAnalyticsProviderFileLogging::SetAge(int InAge)
+{
+	Age = InAge;
+}
+
+void FAnalyticsProviderFileLogging::SetLocation(const FString& InLocation)
+{
+	Location = InLocation;
+}
+
+void FAnalyticsProviderFileLogging::SetGender(const FString& InGender)
+{
+	Gender = InGender;
+}
+
+void FAnalyticsProviderFileLogging::SetBuildInfo(const FString& InBuildInfo)
+{
+	BuildInfo = InBuildInfo;
+}
+
+void FAnalyticsProviderFileLogging::RecordError(const FString& Error, const TArray<FAnalyticsEventAttribute>& Attributes)
+{
+	if (bHasSessionStarted)
+	{
+		check(FileArchive != nullptr);
+
+		if (bHasWrittenFirstEvent)
+		{
+			FileArchive->Logf(TEXT("\t\t,"));
+		}
+		bHasWrittenFirstEvent = true;
+
+		FileArchive->Logf(TEXT("\t\t{"));
+		FileArchive->Logf(TEXT("\t\t\t\"error\" : \"%s\","), *Error);
+
+		FileArchive->Logf(TEXT("\t\t\t\"attributes\" :"));
+		FileArchive->Logf(TEXT("\t\t\t["));
+		bool bHasWrittenFirstAttr = false;
+		// Write out the list of attributes as an array of attribute objects
+		for (auto Attr : Attributes)
+		{
+			if (bHasWrittenFirstAttr)
+			{
+				FileArchive->Logf(TEXT("\t\t\t,"));
+			}
+			FileArchive->Logf(TEXT("\t\t\t{"));
+			FileArchive->Logf(TEXT("\t\t\t\t\"name\" : \"%s\","), *Attr.AttrName);
+			FileArchive->Logf(TEXT("\t\t\t\t\"value\" : \"%s\""), *Attr.AttrValue);
+			FileArchive->Logf(TEXT("\t\t\t}"));
+			bHasWrittenFirstAttr = true;
+		}
+		FileArchive->Logf(TEXT("\t\t\t]"));
+
+		FileArchive->Logf(TEXT("\t\t}"));
+
+		UE_LOG(LogFileLoggingAnalytics, Display, TEXT("Error is (%s) number of attributes is (%d)"), *Error, Attributes.Num());
+	}
+	else
+	{
+		UE_LOG(LogFileLoggingAnalytics, Warning, TEXT("FAnalyticsProviderFileLogging::RecordError called before StartSession. Ignoring."));
+	}
+}
+
+void FAnalyticsProviderFileLogging::RecordProgress(const FString& ProgressType, const FString& ProgressName, const TArray<FAnalyticsEventAttribute>& Attributes)
+{
+	if (bHasSessionStarted)
+	{
+		check(FileArchive != nullptr);
+
+		if (bHasWrittenFirstEvent)
+		{
+			FileArchive->Logf(TEXT("\t\t,"));
+		}
+		bHasWrittenFirstEvent = true;
+
+		FileArchive->Logf(TEXT("\t\t{"));
+		FileArchive->Logf(TEXT("\t\t\t\"eventType\" : \"Progress\","));
+		FileArchive->Logf(TEXT("\t\t\t\"progressType\" : \"%s\","), *ProgressType);
+		FileArchive->Logf(TEXT("\t\t\t\"progressName\" : \"%s\","), *ProgressName);
+
+		FileArchive->Logf(TEXT("\t\t\t\"attributes\" :"));
+		FileArchive->Logf(TEXT("\t\t\t["));
+		bool bHasWrittenFirstAttr = false;
+		// Write out the list of attributes as an array of attribute objects
+		for (auto Attr : Attributes)
+		{
+			if (bHasWrittenFirstAttr)
+			{
+				FileArchive->Logf(TEXT("\t\t\t,"));
+			}
+			FileArchive->Logf(TEXT("\t\t\t{"));
+			FileArchive->Logf(TEXT("\t\t\t\t\"name\" : \"%s\","), *Attr.AttrName);
+			FileArchive->Logf(TEXT("\t\t\t\t\"value\" : \"%s\""), *Attr.AttrValue);
+			FileArchive->Logf(TEXT("\t\t\t}"));
+			bHasWrittenFirstAttr = true;
+		}
+		FileArchive->Logf(TEXT("\t\t\t]"));
+
+		FileArchive->Logf(TEXT("\t\t}"));
+
+		UE_LOG(LogFileLoggingAnalytics, Display, TEXT("Progress event is type (%s), named (%s), number of attributes is (%d)"), *ProgressType, *ProgressName, Attributes.Num());
+	}
+	else
+	{
+		UE_LOG(LogFileLoggingAnalytics, Warning, TEXT("FAnalyticsProviderFileLogging::RecordProgress called before StartSession. Ignoring."));
+	}
+}
+
+void FAnalyticsProviderFileLogging::RecordItemPurchase(const FString& ItemId, int ItemQuantity, const TArray<FAnalyticsEventAttribute>& Attributes)
+{
+	if (bHasSessionStarted)
+	{
+		check(FileArchive != nullptr);
+
+		if (bHasWrittenFirstEvent)
+		{
+			FileArchive->Logf(TEXT("\t\t,"));
+		}
+		bHasWrittenFirstEvent = true;
+
+		FileArchive->Logf(TEXT("\t\t{"));
+		FileArchive->Logf(TEXT("\t\t\t\"eventType\" : \"ItemPurchase\","));
+		FileArchive->Logf(TEXT("\t\t\t\"itemId\" : \"%s\","), *ItemId);
+		FileArchive->Logf(TEXT("\t\t\t\"itemQuantity\" : %d,"), ItemQuantity);
+
+		FileArchive->Logf(TEXT("\t\t\t\"attributes\" :"));
+		FileArchive->Logf(TEXT("\t\t\t["));
+		bool bHasWrittenFirstAttr = false;
+		// Write out the list of attributes as an array of attribute objects
+		for (auto Attr : Attributes)
+		{
+			if (bHasWrittenFirstAttr)
+			{
+				FileArchive->Logf(TEXT("\t\t\t,"));
+			}
+			FileArchive->Logf(TEXT("\t\t\t{"));
+			FileArchive->Logf(TEXT("\t\t\t\t\"name\" : \"%s\","), *Attr.AttrName);
+			FileArchive->Logf(TEXT("\t\t\t\t\"value\" : \"%s\""), *Attr.AttrValue);
+			FileArchive->Logf(TEXT("\t\t\t}"));
+			bHasWrittenFirstAttr = true;
+		}
+		FileArchive->Logf(TEXT("\t\t\t]"));
+
+		FileArchive->Logf(TEXT("\t\t}"));
+
+		UE_LOG(LogFileLoggingAnalytics, Display, TEXT("Item purchase id (%s), quantity (%d), number of attributes is (%d)"), *ItemId, ItemQuantity, Attributes.Num());
+	}
+	else
+	{
+		UE_LOG(LogFileLoggingAnalytics, Warning, TEXT("FAnalyticsProviderFileLogging::RecordItemPurchase called before StartSession. Ignoring."));
+	}
+}
+
+void FAnalyticsProviderFileLogging::RecordCurrencyPurchase(const FString& GameCurrencyType, int GameCurrencyAmount, const TArray<FAnalyticsEventAttribute>& Attributes)
+{
+	if (bHasSessionStarted)
+	{
+		check(FileArchive != nullptr);
+
+		if (bHasWrittenFirstEvent)
+		{
+			FileArchive->Logf(TEXT("\t\t,"));
+		}
+		bHasWrittenFirstEvent = true;
+
+		FileArchive->Logf(TEXT("\t\t{"));
+		FileArchive->Logf(TEXT("\t\t\t\"eventType\" : \"CurrencyPurchase\","));
+		FileArchive->Logf(TEXT("\t\t\t\"gameCurrencyType\" : \"%s\","), *GameCurrencyType);
+		FileArchive->Logf(TEXT("\t\t\t\"gameCurrencyAmount\" : %d,"), GameCurrencyAmount);
+
+		FileArchive->Logf(TEXT("\t\t\t\"attributes\" :"));
+		FileArchive->Logf(TEXT("\t\t\t["));
+		bool bHasWrittenFirstAttr = false;
+		// Write out the list of attributes as an array of attribute objects
+		for (auto Attr : Attributes)
+		{
+			if (bHasWrittenFirstAttr)
+			{
+				FileArchive->Logf(TEXT("\t\t\t,"));
+			}
+			FileArchive->Logf(TEXT("\t\t\t{"));
+			FileArchive->Logf(TEXT("\t\t\t\t\"name\" : \"%s\","), *Attr.AttrName);
+			FileArchive->Logf(TEXT("\t\t\t\t\"value\" : \"%s\""), *Attr.AttrValue);
+			FileArchive->Logf(TEXT("\t\t\t}"));
+			bHasWrittenFirstAttr = true;
+		}
+		FileArchive->Logf(TEXT("\t\t\t]"));
+
+		FileArchive->Logf(TEXT("\t\t}"));
+
+		UE_LOG(LogFileLoggingAnalytics, Display, TEXT("Currency purchase type (%s), quantity (%d), number of attributes is (%d)"), *GameCurrencyType, GameCurrencyAmount, Attributes.Num());
+	}
+	else
+	{
+		UE_LOG(LogFileLoggingAnalytics, Warning, TEXT("FAnalyticsProviderFileLogging::RecordCurrencyPurchase called before StartSession. Ignoring."));
+	}
+}
+
+void FAnalyticsProviderFileLogging::RecordCurrencyGiven(const FString& GameCurrencyType, int GameCurrencyAmount, const TArray<FAnalyticsEventAttribute>& Attributes)
+{
+	if (bHasSessionStarted)
+	{
+		check(FileArchive != nullptr);
+
+		if (bHasWrittenFirstEvent)
+		{
+			FileArchive->Logf(TEXT("\t\t,"));
+		}
+		bHasWrittenFirstEvent = true;
+
+		FileArchive->Logf(TEXT("\t\t{"));
+		FileArchive->Logf(TEXT("\t\t\t\"eventType\" : \"CurrencyGiven\","));
+		FileArchive->Logf(TEXT("\t\t\t\"gameCurrencyType\" : \"%s\","), *GameCurrencyType);
+		FileArchive->Logf(TEXT("\t\t\t\"gameCurrencyAmount\" : %d,"), GameCurrencyAmount);
+
+		FileArchive->Logf(TEXT("\t\t\t\"attributes\" :"));
+		FileArchive->Logf(TEXT("\t\t\t["));
+		bool bHasWrittenFirstAttr = false;
+		// Write out the list of attributes as an array of attribute objects
+		for (auto Attr : Attributes)
+		{
+			if (bHasWrittenFirstAttr)
+			{
+				FileArchive->Logf(TEXT("\t\t\t,"));
+			}
+			FileArchive->Logf(TEXT("\t\t\t{"));
+			FileArchive->Logf(TEXT("\t\t\t\t\"name\" : \"%s\","), *Attr.AttrName);
+			FileArchive->Logf(TEXT("\t\t\t\t\"value\" : \"%s\""), *Attr.AttrValue);
+			FileArchive->Logf(TEXT("\t\t\t}"));
+			bHasWrittenFirstAttr = true;
+		}
+		FileArchive->Logf(TEXT("\t\t\t]"));
+
+		FileArchive->Logf(TEXT("\t\t}"));
+
+		UE_LOG(LogFileLoggingAnalytics, Display, TEXT("Currency given type (%s), quantity (%d), number of attributes is (%d)"), *GameCurrencyType, GameCurrencyAmount, Attributes.Num());
+	}
+	else
+	{
+		UE_LOG(LogFileLoggingAnalytics, Warning, TEXT("FAnalyticsProviderFileLogging::RecordCurrencyGiven called before StartSession. Ignoring."));
+	}
+}
diff --git a/Engine/Plugins/Runtime/Analytics/FileLogging/Source/FileLogging/Private/FileLoggingProvider.h b/Engine/Plugins/Runtime/Analytics/FileLogging/Source/FileLogging/Private/FileLoggingProvider.h
--- a/Engine/Plugins/Runtime/Analytics/FileLogging/Source/FileLogging/Private/FileLoggingProvider.h
+++ b/Engine/Plugins/Runtime/Analytics/FileLogging/Source/FileLogging/Private/FileLoggingProvider.h
@@ -15,6 +15,14 @@ class FAnalyticsProviderFileLogging :
 	FString UserId;
 	/** Unique Id representing the session the analytics are recording for */
 	FString SessionId;
+	/** Holds the Age if set */
+	int32 Age;
+	/** Holds the Location of the user if set */
+	FString Location;
+	/** Holds the Gender of the user if set */
+	FString Gender;
+	/** Holds the build info if set */
+	FString BuildInfo;
 	/** The file archive used to write the data */
 	FArchive* FileArchive;
 
@@ -35,6 +43,19 @@ public:
 	virtual void RecordEvent(const FString& EventName, const TArray<FAnalyticsEventAttribute>& Attributes) override;
 
 	virtual void RecordItemPurchase(const FString& ItemId, const FString& Currency, int PerItemCost, int ItemQuantity) override;
+
 	virtual void RecordCurrencyPurchase(const FString& GameCurrencyType, int GameCurrencyAmount, const FString& RealCurrencyType, float RealMoneyCost, const FString& PaymentProvider) override;
+
 	virtual void RecordCurrencyGiven(const FString& GameCurrencyType, int GameCurrencyAmount) override;
+
+	virtual void SetBuildInfo(const FString& InBuildInfo) override;
+	virtual void SetGender(const FString& InGender) override;
+	virtual void SetLocation(const FString& InLocation) override;
+	virtual void SetAge(const int32 InAge) override;
+
+	virtual void RecordItemPurchase(const FString& ItemId, int ItemQuantity, const TArray<FAnalyticsEventAttribute>& EventAttrs) override;
+	virtual void RecordCurrencyPurchase(const FString& GameCurrencyType, int GameCurrencyAmount, const TArray<FAnalyticsEventAttribute>& EventAttrs) override;
+	virtual void RecordCurrencyGiven(const FString& GameCurrencyType, int GameCurrencyAmount, const TArray<FAnalyticsEventAttribute>& EventAttrs) override;
+	virtual void RecordError(const FString& Error, const TArray<FAnalyticsEventAttribute>& EventAttrs) override;
+	virtual void RecordProgress(const FString& ProgressType, const FString& ProgressHierarchy, const TArray<FAnalyticsEventAttribute>& EventAttrs) override;
 };
diff --git a/Engine/Source/Runtime/Analytics/AnalyticsMulticast/Private/AnalyticsMulticast.cpp b/Engine/Source/Runtime/Analytics/AnalyticsMulticast/Private/AnalyticsMulticast.cpp
--- a/Engine/Source/Runtime/Analytics/AnalyticsMulticast/Private/AnalyticsMulticast.cpp
+++ b/Engine/Source/Runtime/Analytics/AnalyticsMulticast/Private/AnalyticsMulticast.cpp
@@ -36,6 +36,19 @@ public:
 	virtual ~FAnalyticsProviderMulticast();
 
 	bool HasValidProviders() const { return Providers.Num() > 0; }
+
+	virtual void SetBuildInfo(const FString& InBuildInfo) override;
+	virtual void SetGender(const FString& InGender) override;
+	virtual void SetLocation(const FString& InLocation) override;
+	virtual void SetAge(const int32 InAge) override;
+
+	virtual void RecordItemPurchase(const FString& ItemId, int ItemQuantity, const TArray<FAnalyticsEventAttribute>& EventAttrs) override;
+	virtual void RecordCurrencyPurchase(const FString& GameCurrencyType, int GameCurrencyAmount, const TArray<FAnalyticsEventAttribute>& EventAttrs) override;
+	virtual void RecordCurrencyGiven(const FString& GameCurrencyType, int GameCurrencyAmount, const TArray<FAnalyticsEventAttribute>& EventAttrs) override;
+	virtual void RecordError(const FString& Error, const TArray<FAnalyticsEventAttribute>& EventAttrs) override;
+	virtual void RecordProgress(const FString& ProgressType, const FString& ProgressHierarchy, const TArray<FAnalyticsEventAttribute>& EventAttrs) override;
+	virtual void RecordProgress(const FString& ProgressType, const TArray<FString>& ProgressHierarchy, const TArray<FAnalyticsEventAttribute>& EventAttrs) override;
+
 private:
 	TArray<TSharedPtr<IAnalyticsProvider> > Providers;
 	TArray<FString> ProviderModules;
@@ -223,3 +236,84 @@ void FAnalyticsProviderMulticast::RecordCurrencyGiven( const FString& GameCurren
 		(*it)->RecordCurrencyGiven(GameCurrencyType, GameCurrencyAmount);
 	}
 }
+
+void FAnalyticsProviderMulticast::SetBuildInfo(const FString& InBuildInfo)
+{
+	for (auto Provider : Providers)
+	{
+		Provider->SetBuildInfo(InBuildInfo);
+	}
+}
+
+void FAnalyticsProviderMulticast::SetGender(const FString& InGender)
+{
+	for (auto Provider : Providers)
+	{
+		Provider->SetGender(InGender);
+	}
+}
+
+void FAnalyticsProviderMulticast::SetLocation(const FString& InLocation)
+{
+	for (auto Provider : Providers)
+	{
+		Provider->SetLocation(InLocation);
+	}
+}
+
+void FAnalyticsProviderMulticast::SetAge(const int32 InAge)
+{
+	for (auto Provider : Providers)
+	{
+		Provider->SetAge(InAge);
+	}
+}
+
+void FAnalyticsProviderMulticast::RecordItemPurchase(const FString& ItemId, int ItemQuantity, const TArray<FAnalyticsEventAttribute>& EventAttrs)
+{
+	for (auto Provider : Providers)
+	{
+		Provider->RecordItemPurchase(ItemId, ItemQuantity, EventAttrs);
+	}
+}
+
+void FAnalyticsProviderMulticast::RecordCurrencyPurchase(const FString& GameCurrencyType, int GameCurrencyAmount, const TArray<FAnalyticsEventAttribute>& EventAttrs)
+{
+	for (auto Provider : Providers)
+	{
+		Provider->RecordCurrencyPurchase(GameCurrencyType, GameCurrencyAmount, EventAttrs);
+	}
+}
+
+void FAnalyticsProviderMulticast::RecordCurrencyGiven(const FString& GameCurrencyType, int GameCurrencyAmount, const TArray<FAnalyticsEventAttribute>& EventAttrs)
+{
+	for (auto Provider : Providers)
+	{
+		Provider->RecordCurrencyGiven(GameCurrencyType, GameCurrencyAmount, EventAttrs);
+	}
+}
+
+void FAnalyticsProviderMulticast::RecordError(const FString& Error, const TArray<FAnalyticsEventAttribute>& EventAttrs)
+{
+	for (auto Provider : Providers)
+	{
+		Provider->RecordError(Error, EventAttrs);
+	}
+}
+
+void FAnalyticsProviderMulticast::RecordProgress(const FString& ProgressType, const FString& ProgressHierarchy, const TArray<FAnalyticsEventAttribute>& EventAttrs)
+{
+	for (auto Provider : Providers)
+	{
+		Provider->RecordProgress(ProgressType, ProgressHierarchy, EventAttrs);
+	}
+}
+
+void FAnalyticsProviderMulticast::RecordProgress(const FString& ProgressType, const TArray<FString>& ProgressHierarchy, const TArray<FAnalyticsEventAttribute>& EventAttrs)
+{
+	for (auto Provider : Providers)
+	{
+		Provider->RecordProgress(ProgressType, ProgressHierarchy, EventAttrs);
+	}
+}
+
