GitHub #9388 : FileLogging analytic nested json support

If there's JSON inside of JSON in a given analytic, simply output it as JSON instead of wrapping it with double-quotes.
Contributed by: mfortin-bhvr

diff --git a/Engine/Plugins/Runtime/Analytics/FileLogging/Source/FileLogging/Private/FileLogging.cpp b/Engine/Plugins/Runtime/Analytics/FileLogging/Source/FileLogging/Private/FileLogging.cpp
--- a/Engine/Plugins/Runtime/Analytics/FileLogging/Source/FileLogging/Private/FileLogging.cpp
+++ b/Engine/Plugins/Runtime/Analytics/FileLogging/Source/FileLogging/Private/FileLogging.cpp
@@ -181,7 +181,14 @@ void FAnalyticsProviderFileLogging::RecordEvent(const FString& EventName, const
 				}
 				FileArchive->Logf(TEXT("\t\t\t{"));
 				FileArchive->Logf(TEXT("\t\t\t\t\"name\" : \"%s\","), *Attr.GetName());
-				FileArchive->Logf(TEXT("\t\t\t\t\"value\" : \"%s\""), *Attr.GetValue());
+				if (Attr.IsJsonFragment())
+				{
+					FileArchive->Logf(TEXT("\t\t\t\t\"value\" : %s"), *Attr.GetValue());
+				}
+				else
+				{
+					FileArchive->Logf(TEXT("\t\t\t\t\"value\" : \"%s\""), *Attr.GetValue());
+				}
 				FileArchive->Logf(TEXT("\t\t\t}"));
 				bHasWrittenFirstAttr = true;
 			}
